Update program targets to BuildSettingsVersion.Latest

#rnx

diff --git a/Engine/Source/Programs/Enterprise/Datasmith/DatasmithSDK/DatasmithSDK.Target.cs b/Engine/Source/Programs/Enterprise/Datasmith/DatasmithSDK/DatasmithSDK.Target.cs
--- a/Engine/Source/Programs/Enterprise/Datasmith/DatasmithSDK/DatasmithSDK.Target.cs
+++ b/Engine/Source/Programs/Enterprise/Datasmith/DatasmithSDK/DatasmithSDK.Target.cs
@@ -17,7 +17,7 @@ public class DatasmithSDKTarget : TargetRules
 	{
 		Type = TargetType.Program;
 		IncludeOrderVersion = EngineIncludeOrderVersion.Latest;
-		DefaultBuildSettings = BuildSettingsVersion.V2;
+		DefaultBuildSettings = BuildSettingsVersion.Latest;
 		SolutionDirectory = "Programs/Datasmith";
 		bBuildInSolutionByDefault = false;
 
diff --git a/Engine/Source/Programs/VirtualProduction/TextureShare/TextureShareSDK.Target.cs b/Engine/Source/Programs/VirtualProduction/TextureShare/TextureShareSDK.Target.cs
--- a/Engine/Source/Programs/VirtualProduction/TextureShare/TextureShareSDK.Target.cs
+++ b/Engine/Source/Programs/VirtualProduction/TextureShare/TextureShareSDK.Target.cs
@@ -47,7 +47,8 @@ public class TextureShareSDKTarget : TargetRules
 	public TextureShareSDKTarget(TargetInfo Target)
 		: base(Target)
 	{
-		DefaultBuildSettings = BuildSettingsVersion.V2;
+		DefaultBuildSettings = BuildSettingsVersion.Latest;
+		IncludeOrderVersion = EngineIncludeOrderVersion.Latest;
 
 		Type = TargetType.Program;
 		LinkType = TargetLinkType.Monolithic;
